Removed references to deprecated StructUtils module

#Synced-CL 52253672

#Synced-CL 52253672

diff --git a/Engine/Source/Programs/MassTests/MassTests.Build.cs b/Engine/Source/Programs/MassTests/MassTests.Build.cs
--- a/Engine/Source/Programs/MassTests/MassTests.Build.cs
+++ b/Engine/Source/Programs/MassTests/MassTests.Build.cs
@@ -25,7 +25,6 @@ public class MassTests : TestModuleRules
 				"Engine",
 				"MassCore",
 				"MassEntity",
-				"StructUtils",
 			}
 		);
 
